UE5.Main:  Fix for color picker closing immediately when environment light mixer window is open.

Tick function for environment light mixer was calling "SDetailsView::SetObject" with a null Object selection every frame, which internally closes any active color pickers.  Now, it only calls that function if the details view in question doesn't already have an empty Object selection.  This was the narrowest change that fixes the bug.  The code that automatically closes color pickers was added to SDetailsView.cpp in CL 16779751, to fix UE-118530, a Pri 1 bug.  There may be an alternate fix that involves refactoring that code, but this seems adequate, and less risky.

#rnx

diff --git a/Engine/Source/Editor/EnvironmentLightingViewer/Private/SEnvironmentLightingViewer.cpp b/Engine/Source/Editor/EnvironmentLightingViewer/Private/SEnvironmentLightingViewer.cpp
--- a/Engine/Source/Editor/EnvironmentLightingViewer/Private/SEnvironmentLightingViewer.cpp
+++ b/Engine/Source/Editor/EnvironmentLightingViewer/Private/SEnvironmentLightingViewer.cpp
@@ -338,7 +338,12 @@ void SEnvironmentLightingViewer::Tick(const FGeometry& AllottedGeometry, const d
 
 	for (int i = NumDetailsView; i < ENVLIGHT_MAX_DETAILSVIEWS; ++i)
 	{
-		DetailsViews[i]->SetObject(nullptr);
+		// If the details view selection is already empty, don't call SetObject again.  Calling SetObject
+		// otherwise closes any active color picker (UE-121571).
+		if (DetailsViews[i]->GetSelectedObjects().Num() > 0)
+		{
+			DetailsViews[i]->SetObject(nullptr);
+		}
 	}
 }
 
