Fixed crash in environmen light mixer.

diff --git a/Engine/Source/Editor/EnvironmentLightingViewer/Private/SEnvironmentLightingViewer.cpp b/Engine/Source/Editor/EnvironmentLightingViewer/Private/SEnvironmentLightingViewer.cpp
--- a/Engine/Source/Editor/EnvironmentLightingViewer/Private/SEnvironmentLightingViewer.cpp
+++ b/Engine/Source/Editor/EnvironmentLightingViewer/Private/SEnvironmentLightingViewer.cpp
@@ -358,7 +358,7 @@ void SEnvironmentLightingViewer::Tick(const FGeometry& AllottedGeometry, const d
 			break;
 		}
 	}
-	AddComponentDetailView(HeightFogComp, "Volumetric Cloud", VolCloudColor);
+	AddComponentDetailView(HeightFogComp, "Height Fog", VolCloudColor);
 	ButtonCreateHeightFog->SetVisibility(HeightFogComp ? EVisibility::Collapsed : EVisibility::Visible);
 
 	for (int i = NumDetailsView; i < ENVLIGHT_MAX_DETAILSVIEWS; ++i)
diff --git a/Engine/Source/Editor/EnvironmentLightingViewer/Private/SEnvironmentLightingViewer.h b/Engine/Source/Editor/EnvironmentLightingViewer/Private/SEnvironmentLightingViewer.h
--- a/Engine/Source/Editor/EnvironmentLightingViewer/Private/SEnvironmentLightingViewer.h
+++ b/Engine/Source/Editor/EnvironmentLightingViewer/Private/SEnvironmentLightingViewer.h
@@ -13,7 +13,7 @@ class STextBlock;
 class SButton;
 struct FPropertyAndParent;
 
-#define ENVLIGHT_MAX_DETAILSVIEWS 5
+#define ENVLIGHT_MAX_DETAILSVIEWS 6
 
 class SEnvironmentLightingViewer : public SCompoundWidget
 {
