Removed create atmospehre ligth with index 0 from the env light panel.

diff --git a/Engine/Source/Editor/EnvironmentLightingViewer/Private/SEnvironmentLightingViewer.cpp b/Engine/Source/Editor/EnvironmentLightingViewer/Private/SEnvironmentLightingViewer.cpp
--- a/Engine/Source/Editor/EnvironmentLightingViewer/Private/SEnvironmentLightingViewer.cpp
+++ b/Engine/Source/Editor/EnvironmentLightingViewer/Private/SEnvironmentLightingViewer.cpp
@@ -67,9 +67,8 @@ void SEnvironmentLightingViewer::Construct(const FArguments& InArgs)
 								SNew(STextBlock)
 								.Text(this, &SEnvironmentLightingViewer::GetSelectedComboBoxDetailFilterTextLabel)
 							];
-	
+
 	uint32 Zero = 0;
-	uint32 One = 1;
 	ButtonCreateSkyLight = SNew(SButton)
 					.HAlign(HAlign_Center)
 					.OnClicked(this, &SEnvironmentLightingViewer::OnButtonCreateSkyLight)
@@ -77,11 +76,7 @@ void SEnvironmentLightingViewer::Construct(const FArguments& InArgs)
 	ButtonCreateAtmosphericLight0 = SNew(SButton)
 					.HAlign(HAlign_Center)
 					.OnClicked(this, &SEnvironmentLightingViewer::OnButtonCreateAtmosphericLight, Zero)
-					.Text(LOCTEXT("CreateAtmosphericLight0", "Create Atmospheric Light 0"));
-	ButtonCreateAtmosphericLight1 = SNew(SButton)
-					.HAlign(HAlign_Center)
-					.OnClicked(this, &SEnvironmentLightingViewer::OnButtonCreateAtmosphericLight, One)
-					.Text(LOCTEXT("CreateAtmosphericLight1", "Create Atmospheric Light 1"));
+					.Text(LOCTEXT("CreateAtmosphericLight0", "Create Atmospheric Light"));
 	ButtonCreateSkyAtmosphere = SNew(SButton)
 					.HAlign(HAlign_Center)
 					.OnClicked(this, &SEnvironmentLightingViewer::OnButtonCreateSkyAtmosphere)
@@ -149,13 +144,6 @@ void SEnvironmentLightingViewer::Construct(const FArguments& InArgs)
 			.Padding(5.0f)
 			.HAlign(HAlign_Center)
 			.VAlign(VAlign_Center)
-			[
-				ButtonCreateAtmosphericLight1->AsShared()
-			]
-			+SWrapBox::Slot()
-			.Padding(5.0f)
-			.HAlign(HAlign_Center)
-			.VAlign(VAlign_Center)
 			[
 				ButtonCreateSkyAtmosphere->AsShared()
 			]
@@ -307,11 +295,8 @@ void SEnvironmentLightingViewer::Tick(const FGeometry& AllottedGeometry, const d
 	};
 
 	UDirectionalLightComponent* AtmosphericLight0 = GetAtmosphericLight(0);
-	UDirectionalLightComponent* AtmosphericLight1 = GetAtmosphericLight(1);
 	AddComponentDetailView(AtmosphericLight0, "Directional Light 0", DirLightColor);
-	AddComponentDetailView(AtmosphericLight1, "Directional Light 1", DirLightColor);
 	ButtonCreateAtmosphericLight0->SetVisibility(AtmosphericLight0 ? EVisibility::Collapsed : EVisibility::Visible);
-	ButtonCreateAtmosphericLight1->SetVisibility(AtmosphericLight1 ? EVisibility::Collapsed : EVisibility::Visible);
 
 	USkyLightComponent* SkyLightComp = nullptr;
 	for (TObjectIterator<USkyLightComponent> ComponentIt; ComponentIt; ++ComponentIt)
diff --git a/Engine/Source/Editor/EnvironmentLightingViewer/Private/SEnvironmentLightingViewer.h b/Engine/Source/Editor/EnvironmentLightingViewer/Private/SEnvironmentLightingViewer.h
--- a/Engine/Source/Editor/EnvironmentLightingViewer/Private/SEnvironmentLightingViewer.h
+++ b/Engine/Source/Editor/EnvironmentLightingViewer/Private/SEnvironmentLightingViewer.h
@@ -13,7 +13,7 @@ class STextBlock;
 class SButton;
 struct FPropertyAndParent;
 
-#define ENVLIGHT_MAX_DETAILSVIEWS 6
+#define ENVLIGHT_MAX_DETAILSVIEWS 5
 
 class SEnvironmentLightingViewer : public SCompoundWidget
 {
@@ -51,7 +51,6 @@ private:
 
 	TSharedPtr<class SButton> ButtonCreateSkyLight;
 	TSharedPtr<class SButton> ButtonCreateAtmosphericLight0;
-	TSharedPtr<class SButton> ButtonCreateAtmosphericLight1;
 	TSharedPtr<class SButton> ButtonCreateSkyAtmosphere;
 	TSharedPtr<class SButton> ButtonCreateVolumetricCloud;
 	TSharedPtr<class SButton> ButtonCreateHeightFog;
