Env. Light Mixer: Expose inner properties of SecondFogData when parent row is displayed

diff --git a/Engine/Source/Editor/EnvironmentLightingViewer/Private/SEnvironmentLightingViewer.cpp b/Engine/Source/Editor/EnvironmentLightingViewer/Private/SEnvironmentLightingViewer.cpp
--- a/Engine/Source/Editor/EnvironmentLightingViewer/Private/SEnvironmentLightingViewer.cpp
+++ b/Engine/Source/Editor/EnvironmentLightingViewer/Private/SEnvironmentLightingViewer.cpp
@@ -556,6 +556,14 @@ bool SEnvironmentLightingViewer::GetIsPropertyVisible(const FPropertyAndParent&
 		}
 		return true;
 	}
+	else if (OwnerStruct && OwnerStruct->GetName().Equals(TEXT("ExponentialHeightFogData")))
+	{
+		if (bShowMinimalOnly)
+		{
+			return false;
+		}
+		return true;
+	}
 	else if (OwnerClass == UExponentialHeightFogComponent::StaticClass())
 	{
 		if (bShowMinimalOnly)
