Environment Light Mixer: Fix incorrect text when creating a directional light

diff --git a/Engine/Source/Editor/EnvironmentLightingViewer/Private/SEnvironmentLightingViewer.cpp b/Engine/Source/Editor/EnvironmentLightingViewer/Private/SEnvironmentLightingViewer.cpp
--- a/Engine/Source/Editor/EnvironmentLightingViewer/Private/SEnvironmentLightingViewer.cpp
+++ b/Engine/Source/Editor/EnvironmentLightingViewer/Private/SEnvironmentLightingViewer.cpp
@@ -78,8 +78,8 @@ void SEnvironmentLightingViewer::Construct(const FArguments& InArgs)
 					.Text(LOCTEXT("CreateSkyLight", "Create Sky Light"));
 	ButtonCreateAtmosphericLight0 = SNew(SButton)
 					.HAlign(HAlign_Center)
-					.OnClicked(this, &SEnvironmentLightingViewer::OnButtonCreateAtmosphericLight, Zero)
-					.Text(LOCTEXT("CreateAtmosphericLight0", "Create Atmospheric Light"));
+					.OnClicked(this, &SEnvironmentLightingViewer::OnButtonCreateDirectionalLight, Zero)
+					.Text(LOCTEXT("CreateAtmosphericLight0", "Create Directional Light"));
 	ButtonCreateSkyAtmosphere = SNew(SButton)
 					.HAlign(HAlign_Center)
 					.OnClicked(this, &SEnvironmentLightingViewer::OnButtonCreateSkyAtmosphere)
@@ -376,7 +376,7 @@ FReply SEnvironmentLightingViewer::OnButtonCreateSkyLight()
 	return FReply::Handled();
 }
 
-FReply SEnvironmentLightingViewer::OnButtonCreateAtmosphericLight(uint32 Index)
+FReply SEnvironmentLightingViewer::OnButtonCreateDirectionalLight(uint32 Index)
 {
 	UWorld* World = GEditor->GetEditorWorldContext().World();
 	if (!World)
diff --git a/Engine/Source/Editor/EnvironmentLightingViewer/Private/SEnvironmentLightingViewer.h b/Engine/Source/Editor/EnvironmentLightingViewer/Private/SEnvironmentLightingViewer.h
--- a/Engine/Source/Editor/EnvironmentLightingViewer/Private/SEnvironmentLightingViewer.h
+++ b/Engine/Source/Editor/EnvironmentLightingViewer/Private/SEnvironmentLightingViewer.h
@@ -70,7 +70,7 @@ private:
 	TSharedPtr<class SButton> ButtonCreateHeightFog;
 
 	FReply OnButtonCreateSkyLight();
-	FReply OnButtonCreateAtmosphericLight(uint32 Index);
+	FReply OnButtonCreateDirectionalLight(uint32 Index);
 	FReply OnButtonCreateSkyAtmosphere();
 	FReply OnButtonCreateVolumetricCloud();
 	FReply OnButtonCreateHeightFog();
