Localised some environment light mixer strings.

#rnx
#tests PC editor

diff --git a/Engine/Source/Editor/EnvironmentLightingViewer/Private/SEnvironmentLightingViewer.cpp b/Engine/Source/Editor/EnvironmentLightingViewer/Private/SEnvironmentLightingViewer.cpp
--- a/Engine/Source/Editor/EnvironmentLightingViewer/Private/SEnvironmentLightingViewer.cpp
+++ b/Engine/Source/Editor/EnvironmentLightingViewer/Private/SEnvironmentLightingViewer.cpp
@@ -55,9 +55,9 @@ void SEnvironmentLightingViewer::Construct(const FArguments& InArgs)
 									.Padding(5.0f)
 									.ToolTipText(LOCTEXT("CheckBoxAtmosphericLightsOnlyToolTip", "Wether to only present atmospheric lights."));
 
-	ComboBoxDetailFilterOptions.Add(MakeShared<FString>(TEXT("Minimal")));
-	ComboBoxDetailFilterOptions.Add(MakeShared<FString>(TEXT("Normal")));
-	ComboBoxDetailFilterOptions.Add(MakeShared<FString>(TEXT("Normal+Advanced")));
+	ComboBoxDetailFilterOptions.Add(MakeShared<FString>(LOCTEXT("FilterMinimal",	"Minimal").ToString()));
+	ComboBoxDetailFilterOptions.Add(MakeShared<FString>(LOCTEXT("FilterNormal",		"Normal").ToString()));
+	ComboBoxDetailFilterOptions.Add(MakeShared<FString>(LOCTEXT("FilterNormalAdv",	"Normal+Advanced").ToString()));
 
 	SelectedComboBoxDetailFilterOptions = 0;
 	ComboBoxDetailFilter =	SNew(SComboBox<TSharedPtr<FString>>)
